Reenable AudioCaptureTimecodeProvider until we have a better solution to replace it.

diff --git a/Engine/Plugins/Media/AudioCaptureTimecodeProvider/Source/AudioCaptureTimecodeProvider/Private/AudioCaptureTimecodeProvider.cpp b/Engine/Plugins/Media/AudioCaptureTimecodeProvider/Source/AudioCaptureTimecodeProvider/Private/AudioCaptureTimecodeProvider.cpp
--- a/Engine/Plugins/Media/AudioCaptureTimecodeProvider/Source/AudioCaptureTimecodeProvider/Private/AudioCaptureTimecodeProvider.cpp
+++ b/Engine/Plugins/Media/AudioCaptureTimecodeProvider/Source/AudioCaptureTimecodeProvider/Private/AudioCaptureTimecodeProvider.cpp
@@ -12,10 +12,10 @@ PRAGMA_DISABLE_DEPRECATION_WARNINGS
 
 /* FLinearTimecodeAudioCaptureCustomTimeStepImplementation implementation
 *****************************************************************************/
-struct UDEPRECATED_AudioCaptureTimecodeProvider::FLinearTimecodeAudioCaptureCustomTimeStepImplementation
+struct UAudioCaptureTimecodeProvider::FLinearTimecodeAudioCaptureCustomTimeStepImplementation
 {
 public:
-	FLinearTimecodeAudioCaptureCustomTimeStepImplementation(UDEPRECATED_AudioCaptureTimecodeProvider* InOwner)
+	FLinearTimecodeAudioCaptureCustomTimeStepImplementation(UAudioCaptureTimecodeProvider* InOwner)
 		: bWarnedAboutTheInvalidAudioChannel(false)
 		, bFrameRateReach0Counter(0)
 		, Owner(InOwner)
@@ -140,15 +140,15 @@ public:
 	FDropTimecode Timecode;
 
 	/** Owner of the implementation */
-	UDEPRECATED_AudioCaptureTimecodeProvider* Owner;
+	UAudioCaptureTimecodeProvider* Owner;
 
 	/** If the Owner requested the implementation to stop processing */
 	volatile bool bStopRequested;
 };
 
-/* UDEPRECATED_AudioCaptureTimecodeProvider
+/* UAudioCaptureTimecodeProvider
 *****************************************************************************/
-UDEPRECATED_AudioCaptureTimecodeProvider::UDEPRECATED_AudioCaptureTimecodeProvider(const FObjectInitializer& ObjectInitializer)
+UAudioCaptureTimecodeProvider::UAudioCaptureTimecodeProvider(const FObjectInitializer& ObjectInitializer)
 	: Super(ObjectInitializer)
 	, AudioChannel(1)
 	, Implementation(nullptr)
@@ -158,12 +158,12 @@ UDEPRECATED_AudioCaptureTimecodeProvider::UDEPRECATED_AudioCaptureTimecodeProvid
 
 /* UTimecodeProvider interface implementation
 *****************************************************************************/
-FQualifiedFrameTime UDEPRECATED_AudioCaptureTimecodeProvider::GetQualifiedFrameTime() const
+FQualifiedFrameTime UAudioCaptureTimecodeProvider::GetQualifiedFrameTime() const
 {
 	return FQualifiedFrameTime(GetTimecodeInternal(), GetFrameRateInternal());
 }
 
-FTimecode UDEPRECATED_AudioCaptureTimecodeProvider::GetTimecodeInternal() const
+FTimecode UAudioCaptureTimecodeProvider::GetTimecodeInternal() const
 {
 	FTimecode Result;
 	{
@@ -186,7 +186,7 @@ FTimecode UDEPRECATED_AudioCaptureTimecodeProvider::GetTimecodeInternal() const
 	return Result;
 }
 
-FFrameRate UDEPRECATED_AudioCaptureTimecodeProvider::GetFrameRateInternal() const
+FFrameRate UAudioCaptureTimecodeProvider::GetFrameRateInternal() const
 {
 	FFrameRate Result = FrameRate;
 	if (bDetectFrameRate)
@@ -225,7 +225,7 @@ FFrameRate UDEPRECATED_AudioCaptureTimecodeProvider::GetFrameRateInternal() cons
 	return Result;
 }
 
-bool UDEPRECATED_AudioCaptureTimecodeProvider::Initialize(class UEngine* InEngine)
+bool UAudioCaptureTimecodeProvider::Initialize(class UEngine* InEngine)
 {
 	check(Implementation == nullptr);
 	delete Implementation; // in case
@@ -242,14 +242,14 @@ bool UDEPRECATED_AudioCaptureTimecodeProvider::Initialize(class UEngine* InEngin
 	return bInitialized;
 }
 
-void UDEPRECATED_AudioCaptureTimecodeProvider::Shutdown(class UEngine* InEngine)
+void UAudioCaptureTimecodeProvider::Shutdown(class UEngine* InEngine)
 {
 	SynchronizationState = ETimecodeProviderSynchronizationState::Closed;
 	delete Implementation;
 	Implementation = nullptr;
 }
 
-void UDEPRECATED_AudioCaptureTimecodeProvider::BeginDestroy()
+void UAudioCaptureTimecodeProvider::BeginDestroy()
 {
 	delete Implementation;
 	Super::BeginDestroy();
diff --git a/Engine/Plugins/Media/AudioCaptureTimecodeProvider/Source/AudioCaptureTimecodeProvider/Public/AudioCaptureTimecodeProvider.h b/Engine/Plugins/Media/AudioCaptureTimecodeProvider/Source/AudioCaptureTimecodeProvider/Public/AudioCaptureTimecodeProvider.h
--- a/Engine/Plugins/Media/AudioCaptureTimecodeProvider/Source/AudioCaptureTimecodeProvider/Public/AudioCaptureTimecodeProvider.h
+++ b/Engine/Plugins/Media/AudioCaptureTimecodeProvider/Source/AudioCaptureTimecodeProvider/Public/AudioCaptureTimecodeProvider.h
@@ -12,8 +12,8 @@
 /**
  * Read the LTC from the audio capture device.
  */
-UCLASS(Blueprintable, editinlinenew, deprecated)
-class AUDIOCAPTURETIMECODEPROVIDER_API UDEPRECATED_AudioCaptureTimecodeProvider : public UTimecodeProvider
+UCLASS(Blueprintable, editinlinenew)
+class AUDIOCAPTURETIMECODEPROVIDER_API UAudioCaptureTimecodeProvider : public UTimecodeProvider
 {
 	GENERATED_UCLASS_BODY()
 
