Fixed AudioCaptureTimecodeProvider crash due to stack overflow when calling GetFrameRate().

diff --git a/Engine/Plugins/Media/AudioCaptureTimecodeProvider/Source/AudioCaptureTimecodeProvider/Private/AudioCaptureTimecodeProvider.cpp b/Engine/Plugins/Media/AudioCaptureTimecodeProvider/Source/AudioCaptureTimecodeProvider/Private/AudioCaptureTimecodeProvider.cpp
--- a/Engine/Plugins/Media/AudioCaptureTimecodeProvider/Source/AudioCaptureTimecodeProvider/Private/AudioCaptureTimecodeProvider.cpp
+++ b/Engine/Plugins/Media/AudioCaptureTimecodeProvider/Source/AudioCaptureTimecodeProvider/Private/AudioCaptureTimecodeProvider.cpp
@@ -180,7 +180,7 @@ FTimecode UAudioCaptureTimecodeProvider::GetTimecodeInternal() const
 	}
 	else
 	{
-		Result.bDropFrameFormat = FTimecode::IsDropFormatTimecodeSupported(GetFrameRate());
+		Result.bDropFrameFormat = FTimecode::IsDropFormatTimecodeSupported(GetFrameRateInternal());
 	}
 
 	return Result;
